ndisplay: added Windows as a supported platform for icvfx and nDisplay templates

diff --git a/Templates/TP_InCamVFXBP/TP_InCamVFXBP.uproject b/Templates/TP_InCamVFXBP/TP_InCamVFXBP.uproject
--- a/Templates/TP_InCamVFXBP/TP_InCamVFXBP.uproject
+++ b/Templates/TP_InCamVFXBP/TP_InCamVFXBP.uproject
@@ -10,6 +10,7 @@
 		}
 	],
 	"TargetPlatforms": [
-		"WindowsNoEditor"
+		"WindowsNoEditor",
+		"Windows"
 	]
 }
\ No newline at end of file
diff --git a/Templates/TP_nDisplayBP/TP_nDisplayBP.uproject b/Templates/TP_nDisplayBP/TP_nDisplayBP.uproject
--- a/Templates/TP_nDisplayBP/TP_nDisplayBP.uproject
+++ b/Templates/TP_nDisplayBP/TP_nDisplayBP.uproject
@@ -32,6 +32,7 @@
 		}
 	],
 	"TargetPlatforms": [
-		"WindowsNoEditor"
+		"WindowsNoEditor",
+		"Windows"
 	]
 }
\ No newline at end of file
