Upgrading ICVFX template and example projects

diff --git a/Templates/TP_InCamVFXBP/TP_InCamVFXBP.uproject b/Templates/TP_InCamVFXBP/TP_InCamVFXBP.uproject
--- a/Templates/TP_InCamVFXBP/TP_InCamVFXBP.uproject
+++ b/Templates/TP_InCamVFXBP/TP_InCamVFXBP.uproject
@@ -24,6 +24,10 @@
 				"Win64",
 				"Linux"
 			]
+		},
+		{
+			"Name": "DisplayClusterLaunch",
+			"Enabled": true
 		}
 	],
 	"TargetPlatforms": [
